Add SecuritySandbox experimental engine plugin which allows a game client to restrict its access rights on Windows to reduce the impact if an attacker is able to take control of it via online play / UGC.

diff --git a/Engine/Plugins/Experimental/SecuritySandbox/SecuritySandbox.uplugin b/Engine/Plugins/Experimental/SecuritySandbox/SecuritySandbox.uplugin
new file mode 100644
--- /dev/null
+++ b/Engine/Plugins/Experimental/SecuritySandbox/SecuritySandbox.uplugin
@@ -0,0 +1,24 @@
+{
+	"FileVersion": 1,
+	"FriendlyName": "Security Sandbox",
+	"Version": 1,
+	"VersionName": "1.0",
+	"CreatedBy": "Epic Games, Inc.",
+	"CreatedByURL": "https://epicgames.com",
+	"Description": "Provides features to help reduce the operating system permissions your game client runs with and therefore reduce the impact to players if an attacker takes control of it through a vulnerability.",
+	"Category": "Security",
+	"EnabledByDefault": false,
+	"IsExperimentalVersion": true,
+	"CanContainContent": false,
+	"Installed": false,
+	"Modules": [
+		{
+			"Name": "SecuritySandbox",
+			"Type": "Runtime",
+			"LoadingPhase": "Default",
+			"PlatformAllowList": [
+				"Win64"
+			]
+		}
+	]
+}
\ No newline at end of file
diff --git a/Engine/Plugins/Experimental/SecuritySandbox/Source/SecuritySandbox/Private/GenericPlatform/GenericPlatformSecuritySandbox.cpp b/Engine/Plugins/Experimental/SecuritySandbox/Source/SecuritySandbox/Private/GenericPlatform/GenericPlatformSecuritySandbox.cpp
new file mode 100644
--- /dev/null
+++ b/Engine/Plugins/Experimental/SecuritySandbox/Source/SecuritySandbox/Private/GenericPlatform/GenericPlatformSecuritySandbox.cpp
@@ -0,0 +1,82 @@
+// Copyright Epic Games, Inc. All Rights Reserved.
+
+#include "GenericPlatformSecuritySandbox.h"
+#include "Misc/App.h"
+#include "Misc/CoreDelegates.h"
+#include "Misc/CommandLine.h"
+
+void FGenericPlatformSecuritySandbox::Init()
+{
+	bIsSandboxEnabled = InitIsSandboxEnabled();
+
+	// Always bind this delegate if the plugin is enabled. When invoked it will check config to determine whether it needs to actually do anything.
+	if (IsEnabled())
+	{
+		FCoreDelegates::OnFEngineLoopInitComplete.AddThreadSafeSP(AsShared(), &FGenericPlatformSecuritySandbox::OnEngineInitComplete);
+	}
+}
+
+void FGenericPlatformSecuritySandbox::RestrictSelf()
+{
+	// Validate config and conditions and only call PlatformRestrictSelf if all looks good.
+	if (!IsEnabled())
+	{
+		UE_LOG(LogSecuritySandbox, Warning, TEXT("SecuritySandbox is disabled by config or command line param. Skipping self restrictions."));
+		return;
+	}
+
+	if (ensureMsgf(!bHasRestrictedSelf, TEXT("RestrictSelf must be called exactly once.")))
+	{
+		UE_LOG(LogSecuritySandbox, Log, TEXT("Restricting self."));
+		PlatformRestrictSelf();
+		bHasRestrictedSelf = true;
+	}
+}
+
+bool FGenericPlatformSecuritySandbox::IsEnabled()
+{
+	return bIsSandboxEnabled;
+}
+
+void FGenericPlatformSecuritySandbox::OnEngineInitComplete()
+{
+	const USecuritySandboxSettings* Settings = GetDefault<USecuritySandboxSettings>();
+	if (Settings && Settings->bAutoRestrictSelf)
+	{
+		UE_LOG(LogSecuritySandbox, Verbose, TEXT("Engine init complete - triggering RestrictSelf."));
+		RestrictSelf();
+	}
+}
+
+bool FGenericPlatformSecuritySandbox::InitIsSandboxEnabled()
+{
+	// Initial value is based on the plugin settings configuration.
+	const USecuritySandboxSettings* Settings = GetDefault<USecuritySandboxSettings>();
+	bool bEnabled = (Settings && Settings->bIsEnabledByDefault);
+
+	// Allow command line override.
+	// -NoSecuritySandbox: Explicitly disable.
+	// -WithSecuritySandbox: Explicitly enable.
+	// If both are specified then WithSecuritySandbox wins.
+	bEnabled = (bEnabled && !FParse::Param(FCommandLine::Get(), TEXT("NoSecuritySandbox")));
+	bEnabled = (bEnabled || FParse::Param(FCommandLine::Get(), TEXT("WithSecuritySandbox")));
+
+	switch (FApp::GetBuildTargetType())
+	{
+		// Client and Game are supported as the main targets for SecuritySandbox.
+		case EBuildTargetType::Client:
+		case EBuildTargetType::Game:
+			break;
+
+		// Editor is normally disabled but may be enabled by a config flag for testing.
+		case EBuildTargetType::Editor:
+			bEnabled = (bEnabled && Settings && Settings->bAllowSandboxInEditor);
+			break;
+
+		// Other targets are not currently supported.
+		default:
+			bEnabled = false;
+			break;
+	}
+	return bEnabled;
+}
diff --git a/Engine/Plugins/Experimental/SecuritySandbox/Source/SecuritySandbox/Private/GenericPlatform/GenericPlatformSecuritySandbox.h b/Engine/Plugins/Experimental/SecuritySandbox/Source/SecuritySandbox/Private/GenericPlatform/GenericPlatformSecuritySandbox.h
new file mode 100644
--- /dev/null
+++ b/Engine/Plugins/Experimental/SecuritySandbox/Source/SecuritySandbox/Private/GenericPlatform/GenericPlatformSecuritySandbox.h
@@ -0,0 +1,39 @@
+// Copyright Epic Games, Inc. All Rights Reserved.
+#pragma once
+
+#include "CoreMinimal.h"
+#include "ISecuritySandbox.h"
+#include "SecuritySandboxSettings.h"
+
+class FGenericPlatformSecuritySandbox : public ISecuritySandbox, public TSharedFromThis<FGenericPlatformSecuritySandbox>
+{
+public:
+	virtual ~FGenericPlatformSecuritySandbox() {};
+	FGenericPlatformSecuritySandbox() {}
+
+	// Initialize the sandbox, will be called exactly once.
+	virtual void Init();
+
+	// Public interface function from ISecuritySandbox that user code can call without validation.
+	virtual void RestrictSelf() override final;
+
+	// True if restriction can be attempted. If false, RestrictSelf does nothing.
+	virtual bool IsEnabled() override final;
+	
+protected:
+	// Implementation function for derived platform classes which we guarantee will be called at most once with configuration options taken into consideration.
+	virtual void PlatformRestrictSelf() {};
+
+private:
+	// True if restriction can be attempted. If false, RestrictSelf does nothing.
+	bool bIsSandboxEnabled = false;
+	
+	// True if RestrictSelf has run. This doesn't convey success or failure, just that it was attempted.
+	bool bHasRestrictedSelf = false;
+
+	// Set the initial value of bIsSandboxEnabled based on config and command line.
+	bool InitIsSandboxEnabled();
+
+	// Used for automatic RestrictSelf if configured in the plugin's settings.
+	void OnEngineInitComplete();
+};
diff --git a/Engine/Plugins/Experimental/SecuritySandbox/Source/SecuritySandbox/Private/PlatformSecuritySandbox.h b/Engine/Plugins/Experimental/SecuritySandbox/Source/SecuritySandbox/Private/PlatformSecuritySandbox.h
new file mode 100644
--- /dev/null
+++ b/Engine/Plugins/Experimental/SecuritySandbox/Source/SecuritySandbox/Private/PlatformSecuritySandbox.h
@@ -0,0 +1,10 @@
+// Copyright Epic Games, Inc. All Rights Reserved.
+#pragma once
+
+#if PLATFORM_WINDOWS
+#include "Windows/WindowsPlatformSecuritySandbox.h"
+typedef FWindowsPlatformSecuritySandbox FPlatformSecuritySandbox;
+#else
+#include "GenericPlatform/GenericPlatformSecuritySandbox.h"
+typedef FGenericPlatformSecuritySandbox FPlatformSecuritySandbox;
+#endif
\ No newline at end of file
diff --git a/Engine/Plugins/Experimental/SecuritySandbox/Source/SecuritySandbox/Private/SecuritySandboxModule.cpp b/Engine/Plugins/Experimental/SecuritySandbox/Source/SecuritySandbox/Private/SecuritySandboxModule.cpp
new file mode 100644
--- /dev/null
+++ b/Engine/Plugins/Experimental/SecuritySandbox/Source/SecuritySandbox/Private/SecuritySandboxModule.cpp
@@ -0,0 +1,53 @@
+// Copyright Epic Games, Inc. All Rights Reserved.
+
+#include "ISecuritySandboxModule.h"
+#include "PlatformSecuritySandbox.h"
+#include "Templates/SharedPointer.h"
+
+DEFINE_LOG_CATEGORY(LogSecuritySandbox);
+
+class FSecuritySandboxModule : public ISecuritySandboxModule
+{
+public:
+	virtual void StartupModule() override;
+	virtual void ShutdownModule() override;
+	virtual void RestrictSelf() override;
+	virtual bool IsEnabled() override;
+	
+private:
+	TSharedPtr<FPlatformSecuritySandbox> PlatformSecuritySandbox;
+};
+
+void FSecuritySandboxModule::StartupModule()
+{
+	if (!PlatformSecuritySandbox)
+	{
+		PlatformSecuritySandbox = MakeShared<FPlatformSecuritySandbox>();
+		PlatformSecuritySandbox->Init();
+	}
+}
+
+void FSecuritySandboxModule::ShutdownModule()
+{
+	PlatformSecuritySandbox = nullptr;
+}
+
+void FSecuritySandboxModule::RestrictSelf()
+{
+	if (!PlatformSecuritySandbox)
+	{
+		// This should never happen.
+		UE_LOG(LogSecuritySandbox, Error, TEXT("RestrictSelf called with no PlatformSecuritySandbox instance available"));
+		checkNoEntry();
+		return;
+	}
+
+	PlatformSecuritySandbox->RestrictSelf();
+}
+
+bool FSecuritySandboxModule::IsEnabled()
+{
+	return PlatformSecuritySandbox->IsEnabled();
+}
+
+IMPLEMENT_MODULE(FSecuritySandboxModule, SecuritySandbox)
diff --git a/Engine/Plugins/Experimental/SecuritySandbox/Source/SecuritySandbox/Private/Windows/WindowsPlatformSecuritySandbox.cpp b/Engine/Plugins/Experimental/SecuritySandbox/Source/SecuritySandbox/Private/Windows/WindowsPlatformSecuritySandbox.cpp
new file mode 100644
--- /dev/null
+++ b/Engine/Plugins/Experimental/SecuritySandbox/Source/SecuritySandbox/Private/Windows/WindowsPlatformSecuritySandbox.cpp
@@ -0,0 +1,321 @@
+// Copyright Epic Games, Inc. All Rights Reserved.
+#include "WindowsPlatformSecuritySandbox.h"
+#include "Misc/CoreDelegates.h"
+#include "Misc/Paths.h"
+#include "Windows/WindowsApplication.h"
+
+// Needed for opening URLs via WinRT APIs when restricted.
+#include "Windows/AllowWindowsPlatformTypes.h"
+	#include <wrl\client.h>
+	#include <wrl\wrappers\corewrappers.h>
+	#include <Windows.System.h>
+	#include <Windows.Foundation.h>
+	#include <processthreadsapi.h>
+#include "Windows/HideWindowsPlatformTypes.h"
+#pragma comment(lib, "windowsapp.lib")
+
+#define UE_LOG_WINAPIFAIL(Verbosity, ApiName) UE_LOG(LogSecuritySandbox, Verbosity, ApiName TEXT(" failed with %i"), GetLastError())
+
+struct SidContainer
+{
+	SidContainer() : Sid(reinterpret_cast<SID*>(Buffer)), Size(sizeof(Buffer))
+	{
+		FMemory::Memzero(Buffer);
+	}
+	SID *const Sid;
+	const size_t Size;
+
+private:
+	uint8 Buffer[SECURITY_MAX_SID_SIZE];
+};
+
+void FWindowsPlatformSecuritySandbox::PlatformRestrictSelf()
+{
+	ApplyJobRestrictions();
+	SetProcessMitigations();
+	SetLowIntegrityLevel();
+
+	if (bNeedsCustomHandlerForURL)
+	{
+		if (ensure(!FCoreDelegates::LaunchCustomHandlerForURL.IsBound()))
+		{
+			FCoreDelegates::LaunchCustomHandlerForURL.BindThreadSafeSP(SharedThis(this), &FWindowsPlatformSecuritySandbox::CustomHandlerForURL);
+			UE_LOG(LogSecuritySandbox, Verbose, TEXT("Registered FCoreDelegates::LaunchCustomHandlerForURL"));
+		}
+		else
+		{
+			UE_LOG(LogSecuritySandbox, Warning, TEXT("FCoreDelegates::LaunchCustomHandlerForURL is already bound. Skipping."));
+		}
+	}
+}
+
+void FWindowsPlatformSecuritySandbox::SetLowIntegrityLevel()
+{
+	const USecuritySandboxSettings* Settings = GetDefault<USecuritySandboxSettings>();
+	if (Settings)
+	{
+		if (Settings->bUseLowIntegrityLevel)
+		{
+			if (FPaths::ShouldSaveToUserDir())
+			{
+				// Make sure that FWindowsPlatformProcess is configured for low integrity mode. Otherwise settings, logs, temporary files, etc will most likely have been
+				// created in directories that we will lose write access to causing all sorts of problems. The FWindowsPlatformProcess integrity config can be set via a
+				// preprocessor flag in Target.cs for engine source builds or a command line parameter at runtime.
+				ensureMsgf(FWindowsPlatformProcess::ShouldExpectLowIntegrityLevel(), TEXT("Setting low integrity level but FWindowsPlatformProcess is not configured for it. Pass -ExpectLowIntegrityLevel or update Target.cs."));
+
+				// Open our current process token and irreversibly assign it to low integrity level.
+				HANDLE hProcessToken = INVALID_HANDLE_VALUE;
+				if (OpenProcessToken(GetCurrentProcess(), TOKEN_ADJUST_DEFAULT, &hProcessToken))
+				{
+					SidContainer Sid = {};
+					DWORD SidSize = Sid.Size;
+					if (CreateWellKnownSid(WinLowLabelSid, nullptr, Sid.Sid, &SidSize))
+					{
+						TOKEN_MANDATORY_LABEL LowIntegrity = {};
+						LowIntegrity.Label.Sid = Sid.Sid;
+						if (SetTokenInformation(hProcessToken, TokenIntegrityLevel, &LowIntegrity, sizeof(LowIntegrity)))
+						{
+							UE_LOG(LogSecuritySandbox, Log, TEXT("Set low integrity level"));
+
+							// Opening URLs via ShellExecute or CreateProcess fails at low integrity level so we need a custom handler.
+							bNeedsCustomHandlerForURL = true;
+						}
+						else
+						{
+							UE_LOG_WINAPIFAIL(Warning, TEXT("SetTokenInformation"));
+						}
+					}
+					else
+					{
+						UE_LOG_WINAPIFAIL(Warning, TEXT("CreateWellKnownSid"));
+					}
+					CloseHandle(hProcessToken);
+				}
+				else
+				{
+					UE_LOG_WINAPIFAIL(Warning, TEXT("OpenProcessToken"));
+				}
+			}
+			else
+			{
+				// Windows low integrity denies write access to almost all locations except for <User>\AppData\LocalLow, so we can only do it when configured to write to the user directory.
+#if UE_BUILD_SHIPPING
+				UE_LOG(LogSecuritySandbox, Error, TEXT("Low integrity level is configured but ShouldSaveToUserDir returned false. Low integrity level will not be set."));
+#else
+				UE_LOG(LogSecuritySandbox, Log, TEXT("Low integrity level is configured but ShouldSaveToUserDir returned false. Low integrity level will not be set. "
+					"This is normal for Development and Test builds. If you want to test with low integrity, pass -SaveToUserDir on the command line."));
+#endif
+			}
+		}
+		else
+		{
+			UE_LOG(LogSecuritySandbox, Log, TEXT("Low integrity level disabled by config, skipping"));
+		}
+	}
+	else 
+	{
+		UE_LOG(LogSecuritySandbox, Warning, TEXT("Config is not available, skipping integrity level change"));
+	}
+}
+
+void FWindowsPlatformSecuritySandbox::SetProcessMitigations()
+{
+	// These functions were introduced with Windows 8. The UE min spec is Windows 10, but at the time of writing _WIN32_WINNT is still
+	// set to _WIN32_WINNT_WIN7 and changing it just for this would be disruptive, so we can fetch them dynamically instead.
+	typedef bool(WINAPI*pfnGetProcessMitigationPolicy)(HANDLE hProcess, PROCESS_MITIGATION_POLICY MitigationPolicy, PVOID lpBuffer, SIZE_T dwLength);
+	typedef bool(WINAPI*pfnSetProcessMitigationPolicy)(PROCESS_MITIGATION_POLICY MitigationPolicy, PVOID lpBuffer, SIZE_T dwLength);
+
+	const USecuritySandboxSettings* Settings = GetDefault<USecuritySandboxSettings>();
+	if (Settings)
+	{
+		if (Settings->bDisallowLowIntegrityLibraries)
+		{
+			HMODULE hKernel32 = nullptr;
+			if (GetModuleHandleEx(0, TEXT("kernel32.dll"), &hKernel32))
+			{
+				const pfnGetProcessMitigationPolicy fnGetProcessMitigationPolicy = reinterpret_cast<pfnGetProcessMitigationPolicy>((void*)GetProcAddress(hKernel32, "GetProcessMitigationPolicy"));
+				const pfnSetProcessMitigationPolicy fnSetProcessMitigationPolicy = reinterpret_cast<pfnSetProcessMitigationPolicy>((void*)GetProcAddress(hKernel32, "SetProcessMitigationPolicy"));
+				if (fnGetProcessMitigationPolicy && fnSetProcessMitigationPolicy)
+				{
+					PROCESS_MITIGATION_IMAGE_LOAD_POLICY ImageLoadPolicy = {};
+					if (fnGetProcessMitigationPolicy(GetCurrentProcess(), ProcessImageLoadPolicy, &ImageLoadPolicy, sizeof(ImageLoadPolicy)))
+					{
+						if (ImageLoadPolicy.NoLowMandatoryLabelImages != 1)
+						{
+							// Prevent low integrity images from being loaded. This would e.g. block an attacker who has hijacked
+							// the game after it dropped to low integrity from simply writing a DLL file to disk and calling LoadLibrary.
+							ImageLoadPolicy.NoLowMandatoryLabelImages = 1;
+							if (fnSetProcessMitigationPolicy(ProcessImageLoadPolicy, &ImageLoadPolicy, sizeof(ImageLoadPolicy)))
+							{
+								UE_LOG(LogSecuritySandbox, Log, TEXT("Set process mitigations"));
+							}
+							else
+							{
+								UE_LOG_WINAPIFAIL(Warning, TEXT("SetProcessMitigationPolicy"));
+							}
+						}
+						else
+						{
+							UE_LOG(LogSecuritySandbox, Verbose, TEXT("NoLowMandatoryLabelImages already set, skipping"));
+						}
+					}
+					else
+					{
+						UE_LOG_WINAPIFAIL(Warning, TEXT("GetProcessMitigationPolicy"));
+					}
+				}
+				else
+				{
+					UE_LOG(LogSecuritySandbox, Error, TEXT("Failed to resolve process mitigation exports"));
+				}
+				FreeLibrary(hKernel32);
+			}
+			else
+			{
+				UE_LOG(LogSecuritySandbox, Error, TEXT("Failed to find kernel32.dll"));
+			}
+		}
+		else
+		{
+			UE_LOG(LogSecuritySandbox, Verbose, TEXT("Process mitigations disabled by config, skipping"));
+		}
+	}
+	else
+	{
+		UE_LOG(LogSecuritySandbox, Warning, TEXT("Config is not available, skipping process mitigations"));
+	}
+}
+
+void FWindowsPlatformSecuritySandbox::ApplyJobRestrictions()
+{
+	const USecuritySandboxSettings* Settings = GetDefault<USecuritySandboxSettings>();
+	if (Settings)
+	{
+		if (Settings->bDisallowChildProcesses || Settings->bDisallowSystemOperations)
+		{
+			HANDLE hJob = CreateJobObject(nullptr, nullptr);
+			if (hJob)
+			{
+				// Prevent creation of child processes.
+				if (Settings->bDisallowChildProcesses)
+				{
+					JOBOBJECT_BASIC_LIMIT_INFORMATION BasicInfo = {};
+					BasicInfo.ActiveProcessLimit = 1;
+					BasicInfo.LimitFlags = JOB_OBJECT_LIMIT_ACTIVE_PROCESS;
+					if (SetInformationJobObject(hJob, JobObjectBasicLimitInformation, &BasicInfo, sizeof(BasicInfo)))
+					{
+						UE_LOG(LogSecuritySandbox, Log, TEXT("Set job basic limits: 0x%X"), BasicInfo.LimitFlags);
+
+						// UE tries to CreateProcess by default when launching a URL to e.g. open the browser, which will fail if we block child processes.
+						bNeedsCustomHandlerForURL = true;
+					}
+					else
+					{
+						UE_LOG_WINAPIFAIL(Warning, TEXT("SetInformationJobObject Basic"));
+					}
+				}
+
+				// Prevent unnecessary system operations.
+				if (Settings->bDisallowSystemOperations)
+				{
+					JOBOBJECT_BASIC_UI_RESTRICTIONS BasicUi = {};
+					BasicUi.UIRestrictionsClass = JOB_OBJECT_UILIMIT_EXITWINDOWS | JOB_OBJECT_UILIMIT_GLOBALATOMS | JOB_OBJECT_UILIMIT_HANDLES | JOB_OBJECT_UILIMIT_SYSTEMPARAMETERS;
+					if (SetInformationJobObject(hJob, JobObjectBasicUIRestrictions, &BasicUi, sizeof(BasicUi)))
+					{
+						UE_LOG(LogSecuritySandbox, Log, TEXT("Set job UI limits: 0x%X"), BasicUi.UIRestrictionsClass);
+					}
+					else
+					{
+						UE_LOG_WINAPIFAIL(Warning, TEXT("SetInformationJobObject UI"));
+					}
+				}
+
+				// Irreversibly assign ourselves to the job object and close the handle.
+				if (!AssignProcessToJobObject(hJob, GetCurrentProcess()))
+				{
+					UE_LOG_WINAPIFAIL(Warning, TEXT("AssignProcessToJobObject"));
+				}
+				CloseHandle(hJob);
+			}
+			else
+			{
+				UE_LOG_WINAPIFAIL(Warning, TEXT("CreateJobObject"));
+			}
+		}
+		else
+		{
+			UE_LOG(LogSecuritySandbox, Verbose, TEXT("Job restrictions disabled by config, skipping"));
+		}
+	}
+	else
+	{
+		UE_LOG(LogSecuritySandbox, Warning, TEXT("Config is not available, skipping job restrictions"));
+	}
+}
+
+void FWindowsPlatformSecuritySandbox::CustomHandlerForURL(const FString& URL, FString* ErrorMsg)
+{
+	using namespace Microsoft::WRL;
+	using namespace Microsoft::WRL::Wrappers;
+	using namespace ABI::Windows::System;
+	using namespace ABI::Windows::Foundation;
+
+	UE_LOG(LogSecuritySandbox, Log, TEXT("CustomHandlerForURL %s"), *URL);
+
+	FString LocalError;
+	if (FWindowsPlatformMisc::CoInitialize())
+	{
+	    ComPtr<ILauncherStatics> launcherStatics;
+	    HRESULT hr = GetActivationFactory(HStringReference(RuntimeClass_Windows_System_Launcher).Get(), &launcherStatics);   
+	    if (SUCCEEDED(hr))
+	    {
+	        ComPtr<IUriRuntimeClassFactory> uriFactory;
+	        hr = GetActivationFactory(HStringReference(RuntimeClass_Windows_Foundation_Uri).Get(), &uriFactory);
+	        if (SUCCEEDED(hr))
+	        {
+	            HString uriHString;
+	            hr = uriHString.Set(*URL);
+	            if (SUCCEEDED(hr))
+	            {
+	                ComPtr<IUriRuntimeClass> uri;
+	                hr = uriFactory->CreateUri(uriHString.Get(), &uri);
+	                if (SUCCEEDED(hr))
+	                {
+	                    ComPtr<IAsyncOperation<bool>> asyncOp;
+	                    hr = launcherStatics->LaunchUriAsync(uri.Get(), &asyncOp);
+	                    if (!SUCCEEDED(hr)) { LocalError = FString::Printf(TEXT("LaunchUriAsync failed (HRESULT: 0x%08X)"), hr); }
+	                }
+					else
+					{
+						LocalError = FString::Printf(TEXT("Failed to create Uri (HRESULT: 0x%08X)"), hr);
+					}
+	            }
+				else
+				{
+					LocalError = FString::Printf(TEXT("Failed to set HString (HRESULT: 0x%08X)"), hr);
+				}
+	        }
+			else
+			{
+				LocalError = FString::Printf(TEXT("Failed to get IUriRuntimeClassFactory (HRESULT: 0x%08X)"), hr);
+			}
+	    }
+		else
+		{
+			LocalError = FString::Printf(TEXT("Failed to get ILauncherStatics (HRESULT: 0x%08X)"), hr);
+		}
+		FWindowsPlatformMisc::CoUninitialize();
+	}
+	else
+	{
+		LocalError = TEXT("CoInitialize failed");
+	}
+
+	if (!LocalError.IsEmpty())
+	{
+		UE_LOG(LogSecuritySandbox, Warning, TEXT("CustomHandlerForURL failed: %s"), *LocalError);
+	}
+	if (ErrorMsg) { *ErrorMsg = LocalError; }
+}
+
+#undef UE_LOG_WINAPIFAIL
\ No newline at end of file
diff --git a/Engine/Plugins/Experimental/SecuritySandbox/Source/SecuritySandbox/Private/Windows/WindowsPlatformSecuritySandbox.h b/Engine/Plugins/Experimental/SecuritySandbox/Source/SecuritySandbox/Private/Windows/WindowsPlatformSecuritySandbox.h
new file mode 100644
--- /dev/null
+++ b/Engine/Plugins/Experimental/SecuritySandbox/Source/SecuritySandbox/Private/Windows/WindowsPlatformSecuritySandbox.h
@@ -0,0 +1,22 @@
+// Copyright Epic Games, Inc. All Rights Reserved.
+#pragma once
+
+#include "GenericPlatform/GenericPlatformSecuritySandbox.h"
+
+class FWindowsPlatformSecuritySandbox : public FGenericPlatformSecuritySandbox
+{
+public:
+	virtual ~FWindowsPlatformSecuritySandbox() {}
+	virtual void PlatformRestrictSelf() override;
+
+private:
+	void SetLowIntegrityLevel();
+	void SetProcessMitigations();
+	void ApplyJobRestrictions();
+
+	void CustomHandlerForURL(const FString& URL, FString* ErrorMsg);
+
+	typedef FGenericPlatformSecuritySandbox Super;
+
+	bool bNeedsCustomHandlerForURL = false;
+};
diff --git a/Engine/Plugins/Experimental/SecuritySandbox/Source/SecuritySandbox/Public/ISecuritySandbox.h b/Engine/Plugins/Experimental/SecuritySandbox/Source/SecuritySandbox/Public/ISecuritySandbox.h
new file mode 100644
--- /dev/null
+++ b/Engine/Plugins/Experimental/SecuritySandbox/Source/SecuritySandbox/Public/ISecuritySandbox.h
@@ -0,0 +1,33 @@
+// Copyright Epic Games, Inc. All Rights Reserved.
+#pragma once
+
+#include "CoreMinimal.h"
+
+DECLARE_LOG_CATEGORY_EXTERN(LogSecuritySandbox, Log, All);
+
+class ISecuritySandbox
+{
+public:
+	/**
+	 * Permanently restrict this process.
+	 * The exact restrictions applied depend on operating system and configuration.
+	 * 
+	 * Your game can initialize with normal operating system access then call this function before
+	 * activating online features or loading user generated content to reduce risks. USecuritySandboxSettings
+	 * can also be configured to automatically apply restrictions without needing to call this function.
+	 * If SecuritySandbox is disabled by config or command line parameter then this function will do nothing.
+	 * 
+	 * NOTE: State initialized/acquired prior to this call may allow bypassing restrictions,
+	 * e.g. on Windows file handles that were opened beforehand and not closed can still be used.
+	 */
+	virtual void RestrictSelf() = 0;
+
+	/**
+	 * Check if SecuritySandbox as a whole is enabled. If this returns false then
+	 * it may be disabled due to build configuration or command line override.
+	 * Individual components may still be disabled in USecuritySandboxSettings.
+	 * 
+	 * @return True if SecuritySandbox is enabled, otherwise false
+	 */
+	virtual bool IsEnabled() = 0;
+};
diff --git a/Engine/Plugins/Experimental/SecuritySandbox/Source/SecuritySandbox/Public/ISecuritySandboxModule.h b/Engine/Plugins/Experimental/SecuritySandbox/Source/SecuritySandbox/Public/ISecuritySandboxModule.h
new file mode 100644
--- /dev/null
+++ b/Engine/Plugins/Experimental/SecuritySandbox/Source/SecuritySandbox/Public/ISecuritySandboxModule.h
@@ -0,0 +1,56 @@
+// Copyright Epic Games, Inc. All Rights Reserved.
+#pragma once
+
+#include "Modules/ModuleManager.h"
+#include "ISecuritySandbox.h"
+
+/**
+ * Provides features to help minimize the operating system permissions your game client runs with and therefore minimize the impact to players if an attacker takes control of it through a vulnerability.
+ * See ISecuritySandbox for the API.
+ */
+class ISecuritySandboxModule : public IModuleInterface, public ISecuritySandbox
+{
+public:
+	/**
+	 * Singleton-like access to this module's interface.  This is just for convenience!
+	 * Beware of calling this during the shutdown phase, though.  Your module might have been unloaded already.
+	 *
+	 * @return Returns singleton instance, loading the module on demand if needed
+	 */
+	static inline ISecuritySandboxModule& Get()
+	{
+		// Check whether the module is already loaded first. This allows Get() to be called on non-game threads.
+		// LoadModuleChecked is not non-game-thread-safe.
+		ISecuritySandboxModule* ModulePtr = FModuleManager::GetModulePtr<ISecuritySandboxModule>("SecuritySandbox");
+		if (ModulePtr)
+		{
+			return *ModulePtr;
+		}
+
+		return FModuleManager::LoadModuleChecked<ISecuritySandboxModule>("SecuritySandbox");
+	}
+
+	/**
+	 * Checks to see if this module is loaded and ready. It is only valid to call Get() if IsAvailable() returns true.
+	 *
+	 * @return True if the module is loaded and ready to use
+	 */
+	static inline bool IsAvailable()
+	{
+		return FModuleManager::Get().IsModuleLoaded("SecuritySandbox");
+	}
+
+	/**
+	 * Helper for calling ISecuritySandbox::IsEnabled. See that function for details.
+	 *
+	 * @return True if SecuritySandbox is enabled, otherwise false.
+	 */
+	static inline bool IsSandboxEnabled()
+	{
+		return IsAvailable() && Get().IsEnabled();
+	}
+};
+
+#if UE_ENABLE_INCLUDE_ORDER_DEPRECATED_IN_5_2
+#include "CoreMinimal.h"
+#endif
diff --git a/Engine/Plugins/Experimental/SecuritySandbox/Source/SecuritySandbox/Public/SecuritySandboxSettings.h b/Engine/Plugins/Experimental/SecuritySandbox/Source/SecuritySandbox/Public/SecuritySandboxSettings.h
new file mode 100644
--- /dev/null
+++ b/Engine/Plugins/Experimental/SecuritySandbox/Source/SecuritySandbox/Public/SecuritySandboxSettings.h
@@ -0,0 +1,47 @@
+// Copyright Epic Games, Inc. All Rights Reserved.
+#pragma once
+
+#include "Engine/DeveloperSettings.h"
+#include "SecuritySandboxSettings.generated.h"
+
+/**
+ * Implements settings for the Security Sandbox plugin. These are visible in the Project Settings section of the editor.
+ */
+UCLASS(config = Engine, defaultconfig, meta=(DisplayName="Security Sandbox"))
+class SECURITYSANDBOX_API USecuritySandboxSettings : public UDeveloperSettings
+{
+	GENERATED_BODY()
+
+public:
+	USecuritySandboxSettings() {};
+	
+	virtual FName GetCategoryName() const override { return FName(TEXT("Plugins")); }
+
+	/** If true, no command line arguments are needed to activate the sandbox features. If false, -WithSecuritySandbox must be passed to activate them. */
+	UPROPERTY(EditAnywhere, config, Category = General)
+	bool bIsEnabledByDefault = true;
+
+	/** If true, the game will automatically restrict itself after engine initialization completes. Otherwise you must explicitly call ISecuritySandboxModule::RestrictSelf at an appropriate time. */
+	UPROPERTY(EditAnywhere, config, Category = General)
+	bool bAutoRestrictSelf = true;
+
+	/** If true, the sandbox will attempt to operate in editor builds. This may be useful for debugging an uncooked game build but many features of the editor itself will not work. You must restart the editor instance for changes to take effect. */
+	UPROPERTY(EditAnywhere, config, Category = General)
+	bool bAllowSandboxInEditor = false;
+
+	/** If true, the game will be reduced to low integrity level after restricting itself. This means most files and registry keys will become read only in addition to other limitations. See Microsoft documentation for details. */
+	UPROPERTY(EditAnywhere, config, Category = Windows)
+	bool bUseLowIntegrityLevel = true;
+
+	/** If true, the game will be prevented from loading low integrity level libraries after restricting itself. This blocks loading of any executable files written to disk by a low integrity level process. */
+	UPROPERTY(EditAnywhere, config, Category = Windows)
+	bool bDisallowLowIntegrityLibraries = true;
+	
+	/** If true, the game will be prevented from spawning child processes after restricting itself. Currently not compatible with Unreal's default crash reporter. */
+	UPROPERTY(EditAnywhere, config, Category = Windows)
+	bool bDisallowChildProcesses = false;
+
+	/** If true, the game will be prevented from performing system operations that are usually unnecessary like logging the user out of Windows after restricting itself. This unfortunately also blocks clipboard paste from other applications. */
+	UPROPERTY(EditAnywhere, config, Category = Windows)
+	bool bDisallowSystemOperations = false;
+};
diff --git a/Engine/Plugins/Experimental/SecuritySandbox/Source/SecuritySandbox/SecuritySandbox.Build.cs b/Engine/Plugins/Experimental/SecuritySandbox/Source/SecuritySandbox/SecuritySandbox.Build.cs
new file mode 100644
--- /dev/null
+++ b/Engine/Plugins/Experimental/SecuritySandbox/Source/SecuritySandbox/SecuritySandbox.Build.cs
@@ -0,0 +1,29 @@
+// Copyright Epic Games, Inc. All Rights Reserved.
+
+using System;
+
+namespace UnrealBuildTool.Rules
+{
+	public class SecuritySandbox : ModuleRules
+	{
+		public SecuritySandbox(ReadOnlyTargetRules Target) : base(Target)
+		{
+			PublicDependencyModuleNames.AddRange(
+				new string[]
+				{
+					"Core",
+					"CoreUObject",
+					"Engine",
+					"ApplicationCore"
+				}
+			);
+
+			PrivateDependencyModuleNames.AddRange(
+				new string[]
+				{
+					"DeveloperSettings"
+				}
+			);
+		}
+	}
+}
diff --git a/Engine/Source/Runtime/Core/Private/Windows/WindowsPlatformCrashContext.cpp b/Engine/Source/Runtime/Core/Private/Windows/WindowsPlatformCrashContext.cpp
--- a/Engine/Source/Runtime/Core/Private/Windows/WindowsPlatformCrashContext.cpp
+++ b/Engine/Source/Runtime/Core/Private/Windows/WindowsPlatformCrashContext.cpp
@@ -427,6 +427,12 @@ FProcHandle LaunchCrashReportClient(void** OutWritePipe, void** OutReadPipe, uin
 				}
 			}
 		}
+
+		// Tell CRC if we're using low integrity paths (AppData\LocalLow instead of AppData\Local).
+		if (FWindowsPlatformProcess::ShouldExpectLowIntegrityLevel())
+		{
+			FCString::Strncat(CrashReporterClientArgs, TEXT(" -ExpectLowIntegrityLevel"), CR_CLIENT_MAX_ARGS_LEN);
+		}
 	}
 
 #if WITH_EDITOR // Disaster recovery is only enabled for the Editor. Start the server even if in -game, -server, commandlet, the client-side will not connect (its too soon here to query this executable config).
@@ -487,7 +493,8 @@ FProcHandle LaunchCrashReportClient(void** OutWritePipe, void** OutReadPipe, uin
 			FPlatformProcess::CloseProc(Handle);
 
 			// Acquire a handle on the final CRC instance responsible to handle the crash/reports, but forbid this process from terminating it in case we try to terminate it by accident (like a stomped handle that would terminate the wrong process)
-			Handle = RepawnedCrcPid != 0 ? FProcHandle(::OpenProcess(PROCESS_ALL_ACCESS & ~(PROCESS_TERMINATE), 0, RepawnedCrcPid)) : FProcHandle();
+			// Use the minimum required access rights to avoid creating a SecuritySandbox bypass.
+			Handle = RepawnedCrcPid != 0 ? FProcHandle(::OpenProcess(PROCESS_QUERY_LIMITED_INFORMATION | SYNCHRONIZE, 0, RepawnedCrcPid)) : FProcHandle();
 
 			// Update the PID returned to the client.
 			if (OutCrashReportClientProcessId != nullptr)
diff --git a/Engine/Source/Runtime/Core/Private/Windows/WindowsPlatformProcess.cpp b/Engine/Source/Runtime/Core/Private/Windows/WindowsPlatformProcess.cpp
--- a/Engine/Source/Runtime/Core/Private/Windows/WindowsPlatformProcess.cpp
+++ b/Engine/Source/Runtime/Core/Private/Windows/WindowsPlatformProcess.cpp
@@ -1157,17 +1157,31 @@ const TCHAR* FWindowsPlatformProcess::UserTempDir()
 	static FString WindowsUserTempDir;
 	if( !WindowsUserTempDir.Len() )
 	{
-		TCHAR TempPath[MAX_PATH];
-		ZeroMemory(TempPath, sizeof(TCHAR) * MAX_PATH);
+		// Windows temp dir functions don't understand integrity levels so we have to build our own path to AppData\LocalLow\Temp.
+		if (ShouldExpectLowIntegrityLevel())
+		{
+			TCHAR* UserPath;
+			HRESULT Ret = SHGetKnownFolderPath(FOLDERID_LocalAppDataLow, 0, NULL, &UserPath);
+			if (SUCCEEDED(Ret))
+			{
+				WindowsUserTempDir = FString(UserPath).Replace(TEXT("\\"), TEXT("/")) + TEXT("/Temp/");
+				CoTaskMemFree(UserPath);
+			}
+		}
+		else
+		{
+			TCHAR TempPath[MAX_PATH];
+			ZeroMemory(TempPath, sizeof(TCHAR) * MAX_PATH);
 
-		::GetTempPath(MAX_PATH, TempPath);
+			::GetTempPath(MAX_PATH, TempPath);
 
-		// Always expand the temp path in case windows returns short directory names.
-		TCHAR FullTempPath[MAX_PATH];
-		ZeroMemory(FullTempPath, sizeof(TCHAR) * MAX_PATH);
-		::GetLongPathName(TempPath, FullTempPath, MAX_PATH);
+			// Always expand the temp path in case windows returns short directory names.
+			TCHAR FullTempPath[MAX_PATH];
+			ZeroMemory(FullTempPath, sizeof(TCHAR) * MAX_PATH);
+			::GetLongPathName(TempPath, FullTempPath, MAX_PATH);
 
-		WindowsUserTempDir = FString(FullTempPath).Replace(TEXT("\\"), TEXT("/"));
+			WindowsUserTempDir = FString(FullTempPath).Replace(TEXT("\\"), TEXT("/"));
+		}
 	}
 	return *WindowsUserTempDir;
 }
@@ -1179,6 +1193,25 @@ const TCHAR* FWindowsPlatformProcess::UserSettingsDir()
 	{
 		TCHAR* UserPath;
 
+		// get the local or locallow AppData directory depending on integrity configuration
+		HRESULT Ret = SHGetKnownFolderPath(ShouldExpectLowIntegrityLevel() ? FOLDERID_LocalAppDataLow : FOLDERID_LocalAppData, 0, NULL, &UserPath);
+		if (SUCCEEDED(Ret))
+		{
+			// make the base user dir path
+			WindowsUserSettingsDir = FString(UserPath).Replace(TEXT("\\"), TEXT("/")) + TEXT("/");
+			CoTaskMemFree(UserPath);
+		}
+	}
+	return *WindowsUserSettingsDir;
+}
+
+const TCHAR* FWindowsPlatformProcess::UserSettingsDirMediumIntegrity()
+{
+	static FString WindowsUserSettingsDir;
+	if (!WindowsUserSettingsDir.Len())
+	{
+		TCHAR* UserPath;
+
 		// get the local AppData directory
 		HRESULT Ret = SHGetKnownFolderPath(FOLDERID_LocalAppData, 0, NULL, &UserPath);
 		if (SUCCEEDED(Ret))
@@ -1231,7 +1264,7 @@ FString FWindowsPlatformProcess::GetApplicationSettingsDir(const ApplicationSett
 		}
 		case ApplicationSettingsContext::Context::LocalUser:
 		{
-			const HRESULT Ret = SHGetKnownFolderPath(FOLDERID_LocalAppData, 0, NULL, &ApplicationSettingsPath);
+			const HRESULT Ret = SHGetKnownFolderPath(ShouldExpectLowIntegrityLevel() ? FOLDERID_LocalAppDataLow : FOLDERID_LocalAppData, 0, NULL, &ApplicationSettingsPath);
 			if (SUCCEEDED(Ret))
 			{
 				WindowsApplicationSettingsDir = FString(ApplicationSettingsPath);
@@ -2310,6 +2343,31 @@ void FWindowsPlatformProcess::CeaseBeingFirstInstance()
 	WindowsPlatformProcess::ReleaseNamedMutex();
 }
 
+static bool InitShouldExpectLowIntegrityLevel()
+{
+	// Set default based on preprocessor flag, but the behavior can be overridden on the command line at runtime.
+	bool Ret = WINDOWS_LOWINTEGRITYLEVEL_EXPECT_DEFAULT;
+
+	// -ExpectLowIntegrityLevel: Force low integrity level config.
+	// -ExpectMediumIntegrityLevel: Force Windows default medium integrity level config.
+	// If both options are specified then low integrity is selected.
+	if (FParse::Param(::GetCommandLineW(), TEXT("ExpectLowIntegrityLevel")))
+	{
+		Ret = true;
+	}
+	else if (FParse::Param(::GetCommandLineW(), TEXT("ExpectMediumIntegrityLevel")))
+	{
+		Ret = false;
+	}
+	return Ret;
+}
+
+bool FWindowsPlatformProcess::ShouldExpectLowIntegrityLevel()
+{
+	static bool bExpectLowIntegrityLevel = InitShouldExpectLowIntegrityLevel();
+	return bExpectLowIntegrityLevel;
+}
+
 namespace WindowsPlatformProcessImpl
 {
 	static void SetThreadName(LPCSTR ThreadName)
diff --git a/Engine/Source/Runtime/Core/Public/Windows/WindowsPlatform.h b/Engine/Source/Runtime/Core/Public/Windows/WindowsPlatform.h
--- a/Engine/Source/Runtime/Core/Public/Windows/WindowsPlatform.h
+++ b/Engine/Source/Runtime/Core/Public/Windows/WindowsPlatform.h
@@ -207,3 +207,14 @@ typedef FWindowsPlatformTypes FPlatformTypes;
 #ifndef ENABLE_WIN_ALLOC_TRACKING
 #define ENABLE_WIN_ALLOC_TRACKING 0
 #endif
+
+// If set, ShouldExpectLowIntegrityLevel defaults to true. This affects paths for local settings and storage.
+#ifndef WINDOWS_LOWINTEGRITYLEVEL_EXPECT_DEFAULT
+#define WINDOWS_LOWINTEGRITYLEVEL_EXPECT_DEFAULT 0
+#endif
+
+// If set, the engine will attempt to automatically migrate user data like settings from the default medium integrity ProjectUserDir path to the low integrity one.
+// Migration is only attempted if the low integrity ProjectUserDir path is empty, e.g. first run after a game update that sets WINDOWS_LOWINTEGRITYLEVEL_EXPECT_DEFAULT = 1.
+#ifndef WINDOWS_LOWINTEGRITYLEVEL_AUTOMIGRATE_USERDATA
+#define WINDOWS_LOWINTEGRITYLEVEL_AUTOMIGRATE_USERDATA 0
+#endif
\ No newline at end of file
diff --git a/Engine/Source/Runtime/Core/Public/Windows/WindowsPlatformProcess.h b/Engine/Source/Runtime/Core/Public/Windows/WindowsPlatformProcess.h
--- a/Engine/Source/Runtime/Core/Public/Windows/WindowsPlatformProcess.h
+++ b/Engine/Source/Runtime/Core/Public/Windows/WindowsPlatformProcess.h
@@ -140,6 +140,7 @@ public:
 	static CORE_API const TCHAR* UserDir();
 	static CORE_API const TCHAR* UserTempDir();
 	static CORE_API const TCHAR* UserSettingsDir();
+	static CORE_API const TCHAR* UserSettingsDirMediumIntegrity();
 	static CORE_API const TCHAR* ApplicationSettingsDir();
 	static CORE_API FString GetApplicationSettingsDir(const ApplicationSettingsContext& Settings);
 	static CORE_API const TCHAR* ComputerName();
@@ -205,6 +206,15 @@ public:
 	 */
 	static CORE_API void CeaseBeingFirstInstance();
 
+	/**
+	 * Whether to expect to run at a low process integrity level or not. This affects the paths that must be used for user and temp storage.
+	 * The process may launch at default (medium) integrity level and then downgrade itself to low later for security benefits, so checking
+	 * the current integrity level alone is not sufficient.
+	 *
+	 * @return true if low integrity level should be expected, otherwise false.
+	 */
+	static CORE_API bool ShouldExpectLowIntegrityLevel();
+
 protected:
 
 	/**
diff --git a/Engine/Source/Runtime/Engine/Private/GameUserSettings.cpp b/Engine/Source/Runtime/Engine/Private/GameUserSettings.cpp
--- a/Engine/Source/Runtime/Engine/Private/GameUserSettings.cpp
+++ b/Engine/Source/Runtime/Engine/Private/GameUserSettings.cpp
@@ -18,6 +18,11 @@
 #include "HDRHelper.h"
 #include "UnrealClient.h"
 
+// For sandboxing migration of UserSettings from AppData\Local to AppData\LocalLow.
+#if PLATFORM_WINDOWS
+#include "Windows/WindowsPlatformProcess.h"
+#endif
+
 #include UE_INLINE_GENERATED_CPP_BY_NAME(GameUserSettings)
 
 extern EWindowMode::Type GetWindowModeType(EWindowMode::Type WindowMode);
@@ -583,6 +588,45 @@ void UGameUserSettings::SaveSettings()
 
 void UGameUserSettings::LoadConfigIni(bool bForceReload/*=false*/)
 {
+#if PLATFORM_WINDOWS
+	// If configured, try to migrate the user ini file by copying it from AppData\Local to AppData\LocalLow.
+	static const bool bShouldAutomigrateUserData = FPaths::ShouldSaveToUserDir() &&
+		FWindowsPlatformProcess::ShouldExpectLowIntegrityLevel() &&
+		(WINDOWS_LOWINTEGRITYLEVEL_AUTOMIGRATE_USERDATA || FParse::Param(FCommandLine::Get(), TEXT("AutomigrateUserDataToLowIntegrity")));
+	if (bShouldAutomigrateUserData)
+	{
+		static bool bAttemptedSettingsMigration = false;
+		if (!bAttemptedSettingsMigration)
+		{
+			bAttemptedSettingsMigration = true;
+
+			// The low integrity ini file path e.g. AppData\LocalLow\...
+			const FString LowIntegrityPath = FConfigCacheIni::GetDestIniFilename(*GGameUserSettingsIni, nullptr, *FPaths::GeneratedConfigDir());
+			if (LowIntegrityPath.StartsWith(FWindowsPlatformProcess::UserSettingsDir())) // Confirm that assumptions are valid.
+			{
+				// The medium integrity ini file path e.g. AppData\Local\...
+				const FString MediumIntegrityPath = LowIntegrityPath.Replace(FWindowsPlatformProcess::UserSettingsDir(), FWindowsPlatformProcess::UserSettingsDirMediumIntegrity());
+
+				// This file must exist for migration to proceed. We only migrate from medium integrity to low integrity, not the reverse.
+				if (FPaths::FileExists(MediumIntegrityPath))
+				{
+					FDateTime MediumIntegrityTimestamp, LowIntegrityTimestamp;
+					IFileManager::Get().GetTimeStampPair(*MediumIntegrityPath, *LowIntegrityPath, MediumIntegrityTimestamp, LowIntegrityTimestamp);
+					if (MediumIntegrityTimestamp > LowIntegrityTimestamp)
+					{
+						UE_LOG(LogTemp, Log, TEXT("Migrating config file to low integrity. %s to %s"), *MediumIntegrityPath, *LowIntegrityPath);
+						IFileManager::Get().Copy(*LowIntegrityPath, *MediumIntegrityPath, true);
+					}
+					else
+					{
+						UE_LOG(LogTemp, Log, TEXT("Skipping config file migration to low integrity. %s is newer than %s"), *LowIntegrityPath, *MediumIntegrityPath);
+					}
+				}
+			}
+		}
+	}
+#endif
+
 	FConfigContext Context = FConfigContext::ReadIntoGConfig();
 	Context.bForceReload = bForceReload;
 	Context.Load(TEXT("GameUserSettings"), GGameUserSettingsIni);
